JWT plugin, with skeleton for supporting client-side verification, integrated with FOnlineIdentityMcp.

diff --git a/Engine/Plugins/Experimental/JWT/JWT.uplugin b/Engine/Plugins/Experimental/JWT/JWT.uplugin
new file mode 100644
--- /dev/null
+++ b/Engine/Plugins/Experimental/JWT/JWT.uplugin
@@ -0,0 +1,25 @@
+{
+	"FileVersion" : 3,
+	"Version" : 1,
+	"VersionName" : "1.0",
+	"FriendlyName" : "JSON Web Token Plugin",
+	"Description" : "An API for working with JSON Web Token (JWT) data.",
+	"Category" : "Misc",
+	"CreatedBy" : "Epic Games, Inc.",
+	"CreatedByURL" : "http://epicgames.com",
+	"DocsURL" : "",
+	"MarketplaceURL" : "",
+	"SupportURL" : "",
+	"EnabledByDefault" : false,
+	"CanContainContent" : false,
+	"IsBetaVersion" : false,
+	"Installed" : false,
+	"Modules" :
+	[
+		{
+			"Name" : "JWT",
+			"Type" : "Runtime",
+			"LoadingPhase" : "Default"
+		}
+	]
+}
diff --git a/Engine/Plugins/Experimental/JWT/Source/JWT/JWT.Build.cs b/Engine/Plugins/Experimental/JWT/Source/JWT/JWT.Build.cs
new file mode 100644
--- /dev/null
+++ b/Engine/Plugins/Experimental/JWT/Source/JWT/JWT.Build.cs
@@ -0,0 +1,46 @@
+// Copyright Epic Games, Inc. All Rights Reserved.
+
+namespace UnrealBuildTool.Rules
+{
+	public class JWT : ModuleRules
+	{
+		public JWT(ReadOnlyTargetRules Target) : base(Target)
+		{
+			PublicIncludePaths.AddRange(
+				new string[] {
+					// ... add public include paths required here ...
+				}
+			);
+
+			PrivateIncludePaths.AddRange(
+				new string[] {
+					// ... add other private include paths required here ...
+				}
+			);
+
+			PublicDependencyModuleNames.AddRange(
+				new string[]
+				{
+					"Core",
+					// ... add other public dependencies that you statically link with here ...
+					"Engine",
+					"Json",
+				}
+			);
+
+			PrivateDependencyModuleNames.AddRange(
+				new string[]
+				{
+					// ... add private dependencies that you statically link with here ...
+				}
+			);
+
+			DynamicallyLoadedModuleNames.AddRange(
+				new string[]
+				{
+					// ... add any modules that your module loads dynamically here ...
+				}
+			);
+		}
+	}
+}
diff --git a/Engine/Plugins/Experimental/JWT/Source/JWT/Private/IJwt.cpp b/Engine/Plugins/Experimental/JWT/Source/JWT/Private/IJwt.cpp
new file mode 100644
--- /dev/null
+++ b/Engine/Plugins/Experimental/JWT/Source/JWT/Private/IJwt.cpp
@@ -0,0 +1,45 @@
+// Copyright Epic Games, Inc. All Rights Reserved.
+
+#include "IJwt.h"
+#include "CoreMinimal.h"
+#include "Modules/ModuleManager.h"
+
+DEFINE_LOG_CATEGORY(LogJwt);
+
+class FJwt : public IJwt
+{
+	/** IModuleInterface implementation */
+	virtual void StartupModule() override;
+	virtual void ShutdownModule() override;
+};
+
+IMPLEMENT_MODULE(FJwt, JWT)
+
+
+
+void FJwt::StartupModule()
+{
+	// This code will execute after your module is loaded into memory (but after global variables are initialized, of course.)
+}
+
+
+void FJwt::ShutdownModule()
+{
+	// This function may be called during shutdown to clean up your module.  For modules that support dynamic reloading,
+	// we call this function before unloading the module.
+}
+
+
+namespace UE::JWT
+{
+	TOptional<FJsonWebToken> FromString(const FStringView InEncodedJsonWebToken, const bool bIsSignatureEncoded)
+	{
+		return FJsonWebToken::FromString(InEncodedJsonWebToken, bIsSignatureEncoded);
+	}
+
+
+	bool FromString(const FStringView InEncodedJsonWebToken, FJsonWebToken& OutJsonWebToken, const bool bIsSignatureEncoded)
+	{
+		return FJsonWebToken::FromString(InEncodedJsonWebToken, OutJsonWebToken, bIsSignatureEncoded);
+	}
+}
diff --git a/Engine/Plugins/Experimental/JWT/Source/JWT/Private/JsonWebToken.cpp b/Engine/Plugins/Experimental/JWT/Source/JWT/Private/JsonWebToken.cpp
new file mode 100644
--- /dev/null
+++ b/Engine/Plugins/Experimental/JWT/Source/JWT/Private/JsonWebToken.cpp
@@ -0,0 +1,315 @@
+// Copyright Epic Games, Inc. All Rights Reserved.
+
+#include "JsonWebToken.h"
+#include "Algo/Count.h"
+#include "Misc/Base64.h"
+#include "Serialization/JsonSerializer.h"
+
+// JWT header field names
+const TCHAR *const FJsonWebToken::HEADER_TYPE = TEXT("typ");
+const TCHAR *const FJsonWebToken::HEADER_KEY_ID = TEXT("kid");
+const TCHAR *const FJsonWebToken::HEADER_ALGORITHM = TEXT("alg");
+
+// JWT header expected values
+const TCHAR *const FJsonWebToken::TYPE_VALUE_JWT = TEXT("JWT");
+
+// JWT payload claim field names
+const TCHAR *const FJsonWebToken::CLAIM_ISSUER = TEXT("iss");
+const TCHAR *const FJsonWebToken::CLAIM_ISSUED_AT = TEXT("iat");
+const TCHAR *const FJsonWebToken::CLAIM_EXPIRATION = TEXT("exp");
+const TCHAR *const FJsonWebToken::CLAIM_SUBJECT = TEXT("sub");
+const TCHAR *const FJsonWebToken::CLAIM_AUDIENCE = TEXT("aud");
+
+
+namespace
+{
+	/**
+	 * Helper for splitting FStringViews, like FString.
+	 */
+	bool SplitStringView(const FStringView InSourceString, const TCHAR InDelimiter, FStringView& OutLeft, FStringView& OutRight)
+	{
+		int32 Position = INDEX_NONE;
+		if (!InSourceString.FindChar(InDelimiter, Position) || Position == INDEX_NONE)
+		{
+			return false;
+		}
+
+		OutLeft = InSourceString.Left(Position);
+		OutRight = InSourceString.RightChop(Position + 1);
+		return true;
+	}
+
+	/**
+	 * Splits an encoded JWT string into the 3 header, payload, and signature parts.
+	 */
+	bool SplitEncodedJsonWebTokenString(const FStringView InEncodedJsonWebTokenString, FStringView& OutEncodedHeaderPart, FStringView& OutEncodedPayloadPart, FStringView& OutEncodedSignaturePart)
+	{
+		FStringView Rest;
+
+		if (!SplitStringView(InEncodedJsonWebTokenString, TEXT('.'), OutEncodedHeaderPart, Rest))
+		{
+			UE_LOG(LogJwt, Warning, TEXT("[SplitEncodedJsonWebTokenString] Cannot extract header from token string."));
+			return false;
+		}
+
+		if (!SplitStringView(Rest, TEXT('.'), OutEncodedPayloadPart, OutEncodedSignaturePart))
+		{
+			UE_LOG(LogJwt, Warning, TEXT("[SplitEncodedJsonWebTokenString] Cannot extract payload from token string."));
+			return false;
+		}
+
+		if (OutEncodedHeaderPart.IsEmpty() || OutEncodedPayloadPart.IsEmpty())
+		{
+			UE_LOG(LogJwt, Warning, TEXT("[SplitEncodedJsonWebTokenString] Empty header and/or payload in token string."));
+			return false;
+		}
+
+		return true;
+	}
+
+	bool Base64Decode(const FStringView InSource, FString& OutDest)
+	{
+		if (InSource.IsEmpty())
+		{
+			return false;
+		}
+
+		return FBase64::Decode(FString(InSource), OutDest);
+	}
+
+	bool Base64Decode(const FStringView InSource, TArray<uint8>& OutDest)
+	{
+		if (InSource.IsEmpty())
+		{
+			return false;
+		}
+
+		return FBase64::Decode(FString(InSource), OutDest);
+	}
+
+	bool StringViewToBytes(const FStringView In, TArray<uint8>& OutBytes, const bool IsEncoded)
+	{
+		if (IsEncoded)
+		{
+			return Base64Decode(In, OutBytes);
+		}
+
+		OutBytes.Reserve(In.Len());
+		for (const TCHAR& Ch : In)
+		{
+			OutBytes.Add(static_cast<uint8>(Ch));
+		}
+
+		return true;
+	}
+}
+
+void FJsonWebToken::DumpJsonObject(const FJsonObject& InJsonObject)
+{
+	if (!UE_LOG_ACTIVE(LogJwt, VeryVerbose))
+	{
+		return;
+	}
+
+	UE_LOG(LogJwt, VeryVerbose, TEXT("size=%d, fields:"), InJsonObject.Values.Num());
+
+	for (const TMap<FString, TSharedPtr<FJsonValue>>::ElementType& Pair : InJsonObject.Values)
+	{
+		const FString& Name = Pair.Key;
+		const TSharedPtr<FJsonValue>& Value = Pair.Value;
+
+		if (!Value)
+		{
+			UE_LOG(LogJwt, VeryVerbose, TEXT("  %s [undefined]"), *Name);
+			continue;
+		}
+
+		if (Value->IsNull())
+		{
+			UE_LOG(LogJwt, VeryVerbose, TEXT("  %s [null]"), *Name);
+			continue;
+		}
+
+		// Present the JSON value in a printable form; TryGetString() is overridden for
+		// non-compound types (bool, number, string).
+		FString OutString;
+		if (Value->TryGetString(OutString))
+		{
+			UE_LOG(LogJwt, VeryVerbose, TEXT("  %s: %s"), *Name, *OutString);
+			continue;
+		}
+
+		const TArray< TSharedPtr<FJsonValue> >* OutArray = nullptr;
+		if (Value->TryGetArray(OutArray))
+		{
+			UE_LOG(LogJwt, VeryVerbose, TEXT("  %s [array, len=%d]"), *Name, OutArray->Num());
+			continue;
+		}
+
+		const TSharedPtr<FJsonObject>* OutObject = nullptr;
+		if (Value->TryGetObject(OutObject))
+		{
+			const int32 Length = (*OutObject).IsValid() ? (*OutObject)->Values.Num() : 0;
+
+			UE_LOG(LogJwt, VeryVerbose, TEXT("  %s [object, len=%d]"), *Name, Length);
+			continue;
+		}
+
+		UE_LOG(LogJwt, VeryVerbose, TEXT("  %s [unknown type])"), *Name);
+	}
+}
+
+
+TSharedPtr<FJsonObject> FJsonWebToken::FromJson(const FString& InJsonStr)
+{
+	TSharedRef<TJsonReader<> > Reader = TJsonReaderFactory<>::Create(InJsonStr);
+
+	TSharedPtr<FJsonObject> JsonObject;
+	const bool bDeserializeOk = FJsonSerializer::Deserialize(Reader, JsonObject);
+
+	if (!bDeserializeOk)
+	{
+		UE_LOG(LogJwt, VeryVerbose, TEXT("deserialization of JWT failed"));
+	}
+	else if (!JsonObject.IsValid())
+	{
+		UE_LOG(LogJwt, VeryVerbose, TEXT("no custom properties in JWT"));
+	}
+	else
+	{
+		DumpJsonObject(*JsonObject.Get());
+	}
+
+	return JsonObject;
+}
+
+TSharedPtr<FJsonObject> FJsonWebToken::ParseEncodedJson(const FStringView InEncodedJson)
+{
+	TSharedPtr<FJsonObject> ParsedObj;
+
+	FString DecodedJson;
+	if (Base64Decode(InEncodedJson, DecodedJson))
+	{
+		ParsedObj = FromJson(DecodedJson);
+	}
+
+	return ParsedObj;
+}
+
+bool FJsonWebToken::FromString(const FStringView InEncodedJsonWebToken, FJsonWebToken& OutJsonWebToken, const bool bIsSignatureEncoded)
+{
+	TOptional<FJsonWebToken> JsonWebToken = FromString(InEncodedJsonWebToken, bIsSignatureEncoded);
+	if (!JsonWebToken.IsSet())
+	{
+		return false;
+	}
+
+	OutJsonWebToken = MoveTemp(JsonWebToken.GetValue());
+	return true;
+}
+
+TOptional<FJsonWebToken> FJsonWebToken::FromString(const FStringView InEncodedJsonWebToken, const bool bIsSignatureEncoded)
+{
+	// Check for the correct number of dots.
+	if (Algo::Count(InEncodedJsonWebToken, TEXT('.')) != 2)
+	{
+		UE_LOG(LogJwt, Verbose, TEXT("[FJsonWebToken::FromString] Failed. Token is not in {header}.{payload}.[{signature}] form."));
+		return {};
+	}
+
+	// Split out the token string into header, payload, and signature parts.
+	FStringView EncodedHeader;
+	FStringView EncodedPayload;
+	FStringView SignaturePart;
+	if (!SplitEncodedJsonWebTokenString(InEncodedJsonWebToken, EncodedHeader, EncodedPayload, SignaturePart))
+	{
+		return {};
+	}
+
+	// Decode and parse the header.
+	UE_LOG(LogJwt, VeryVerbose, TEXT("[FJsonWebToken::FromString] Parsing JWT header."));
+	const TSharedPtr<FJsonObject> HeaderPtr = ParseEncodedJson(EncodedHeader);
+	if (!HeaderPtr)
+	{
+		UE_LOG(LogJwt, Verbose, TEXT("[FJsonWebToken::FromString] Failed to decode and parse the header."));
+		return {};
+	}
+
+	// Decode and parse the payload.
+	UE_LOG(LogJwt, VeryVerbose, TEXT("[FJsonWebToken::FromString] Parsing JWT payload."));
+	const TSharedPtr<FJsonObject> PayloadPtr = ParseEncodedJson(EncodedPayload);
+	if (!PayloadPtr)
+	{
+		UE_LOG(LogJwt, Verbose, TEXT("[FJsonWebToken::FromString] Failed to decode and parse the payload."));
+		return {};
+	}
+
+	// Decode (but do not parse) the signature if it not empty.
+	TOptional<TArray<uint8>> Signature;
+	if (!SignaturePart.IsEmpty())
+	{
+		UE_LOG(LogJwt, VeryVerbose, TEXT("[FJsonWebToken::FromString] Decoding JWT signature."));
+		TArray<uint8> SignatureBytes;
+		if (!StringViewToBytes(SignaturePart, SignatureBytes, bIsSignatureEncoded))
+		{
+			UE_LOG(LogJwt, Verbose, TEXT("[FJsonWebToken::FromString] Failed to decode the signature."));
+			return {};
+		}
+
+		Signature = MoveTemp(SignatureBytes);
+	}
+
+	const FJsonWebToken JsonWebToken(InEncodedJsonWebToken, HeaderPtr.ToSharedRef(), PayloadPtr.ToSharedRef(), Signature);
+
+	// Validate the type.  If it exists but is not a JWT, then fail.
+	FString Type;
+	if (JsonWebToken.GetType(Type) && (Type != TYPE_VALUE_JWT))
+	{
+		UE_LOG(LogJwt, Verbose, TEXT("[FJsonWebToken::FromString] Failed. Type field exists but is not JWT: \"%s\""), *Type);
+		return {};
+	}
+
+	return JsonWebToken;
+}
+
+FJsonWebToken::FJsonWebToken(
+	const FStringView InEncodedJsonWebToken,
+	const TSharedRef<FJsonObject>& InHeaderPtr,
+	const TSharedRef<FJsonObject>& InPayloadPtr,
+	const TOptional<TArray<uint8>>& InSignature
+)
+	: EncodedJsonWebToken(InEncodedJsonWebToken)
+	, Header(InHeaderPtr)
+	, Payload(InPayloadPtr)
+	, Signature(InSignature)
+{
+}
+
+bool FJsonWebToken::GetType(FString& OutValue) const
+{
+	return Header->TryGetStringField(HEADER_TYPE, OutValue);
+}
+
+bool FJsonWebToken::GetKeyId(FString& OutValue) const
+{
+	return Header->TryGetStringField(HEADER_KEY_ID, OutValue);
+}
+
+bool FJsonWebToken::GetAlgorithm(FString& OutValue) const
+{
+	return Header->TryGetStringField(HEADER_ALGORITHM, OutValue);
+}
+
+bool FJsonWebToken::GetClaim(const FStringView InName, TSharedPtr<FJsonValue>& OutClaim) const
+{
+	OutClaim = Payload->TryGetField(FString(InName));
+	return OutClaim.IsValid();
+}
+
+
+bool FJsonWebToken::Verify() const
+{
+	// Signature verification is yet to be implemented, so return true for now.
+	// TODO: Add support for RSA crypto algorithm (PlatformCrypto plugin) and key management so that signature verification can be performed.
+	return true;
+}
diff --git a/Engine/Plugins/Experimental/JWT/Source/JWT/Public/IJwt.h b/Engine/Plugins/Experimental/JWT/Source/JWT/Public/IJwt.h
new file mode 100644
--- /dev/null
+++ b/Engine/Plugins/Experimental/JWT/Source/JWT/Public/IJwt.h
@@ -0,0 +1,52 @@
+// Copyright Epic Games, Inc. All Rights Reserved.
+
+#pragma once
+
+#include "CoreMinimal.h"
+#include "Modules/ModuleInterface.h"
+#include "Modules/ModuleManager.h"
+#include "JwtIncludes.h"
+
+
+/**
+ * The public interface to this module.  In most cases, this interface is only public to sibling modules 
+ * within this plugin.
+ */
+class IJwt : public IModuleInterface
+{
+
+public:
+	/**
+	 * Singleton-like access to this module's interface.  This is just for convenience!
+	 * Beware of calling this during the shutdown phase, though.  Your module might have been unloaded already.
+	 *
+	 * @return Returns singleton instance, loading the module on demand if needed
+	 */
+	static inline IJwt& Get()
+	{
+		IJwt* const ModulePtr = FModuleManager::GetModulePtr<IJwt>("JWT");
+		if (ModulePtr)
+		{
+			return *ModulePtr;
+		}
+
+		return FModuleManager::LoadModuleChecked<IJwt>("JWT");
+	}
+
+	/**
+	 * Checks to see if this module is loaded and ready.  It is only valid to call Get() if IsAvailable() returns true.
+	 *
+	 * @return True if the module is loaded and ready to use
+	 */
+	static inline bool IsAvailable()
+	{
+		return FModuleManager::Get().IsModuleLoaded("JWT");
+	}
+};
+
+namespace UE::JWT
+{
+	JWT_API TOptional<FJsonWebToken> FromString(const FStringView InEncodedJsonWebToken, const bool bIsSignatureEncoded);
+
+	JWT_API bool FromString(const FStringView InEncodedJWT, FJsonWebToken& OutJsonWebToken, const bool bIsSignatureEncoded);
+}
diff --git a/Engine/Plugins/Experimental/JWT/Source/JWT/Public/JsonWebToken.h b/Engine/Plugins/Experimental/JWT/Source/JWT/Public/JsonWebToken.h
new file mode 100644
--- /dev/null
+++ b/Engine/Plugins/Experimental/JWT/Source/JWT/Public/JsonWebToken.h
@@ -0,0 +1,121 @@
+// Copyright Epic Games, Inc. All Rights Reserved.
+
+#pragma once
+
+#include "CoreMinimal.h"
+#include "JwtGlobals.h"
+#include "Dom/JsonObject.h"
+
+class JWT_API FJsonWebToken
+{
+public:
+	/**
+	 * Creates a JWT from the provided string.
+	 * The string must consist of 3 base64 url encoded parts: a header, payload, and signature.
+	 * The parts must be split by a period character.
+	 * The signature part is optional. If the signature is excluded, the string must still contain a period character in its place.
+	 * Valid formats: "header.payload.signature" and "header.payload."
+	 *
+	 * @param InJsonWebTokenString The string to decode.
+	 * @return An optional set with a FJsonWebToken if the JWT was successfully decoded.
+	 */
+	static TOptional<FJsonWebToken> FromString(const FStringView InEncodedJsonWebToken, const bool bIsSignatureEncoded);
+
+	static bool FromString(const FStringView InEncodedJsonWebToken, FJsonWebToken& OutJsonWebToken, const bool bIsSignatureEncoded);
+
+	/**
+	 * Gets the type.
+	 *
+	 * @param OutValue The value to output on success.
+	 * @return Whether the value exists and was successfully outputted.
+	 */
+	bool GetType(FString& OutValue) const;
+
+	/**
+	 * Gets the key id.
+	 *
+	 * @param OutValue The value to output on success.
+	 * @return Whether the value exists and was successfully outputted.
+	 */
+	bool GetKeyId(FString& OutValue) const;
+
+	/**
+	 * Gets the algorithm that was used to construct the signature.
+	 *
+	 * @param OutValue The value to output on success.
+	 * @return Whether the value exists and was successfully outputted.
+	 */
+	bool GetAlgorithm(FString& OutValue) const;
+
+	/**
+	 * Gets a claim by name.
+	 * This method can be used to get custom claims that are not reserved as part of the JWT specification.
+	 *
+	 * @param InName The name of the claim.
+	 * @param OutClaim The json value to output on success.
+	 * @return Whether the claim exists and was successfully outputted.
+	 */
+	bool GetClaim(const FStringView InName, TSharedPtr<FJsonValue>& OutClaim) const;
+
+	/**
+	 * Gets a claim by name.
+	 * This method can be used to get custom claims that are not reserved as part of the JWT specification.
+	 *
+	 * @param InName The name of the claim.
+	 * @return The json value to output on success, or an invalid shared pointer on failure.
+	 */
+	template<EJson JsonType>
+	TSharedPtr<FJsonValue> GetClaim(const FStringView InName) const
+	{
+		return Payload->GetField<JsonType>(FString(InName));
+	}
+
+	/**
+	 * Perform verification of the encoded header and encoded payload.
+	 * The is done using the cryptographic algorithm specified in the header, and a secret value obtained
+	 * by key-lookup, to generate a signature.  The generated signature must be identical to the signature
+	 * provided in the JWT for verification to succeed.
+	 * 
+	 * @return true if the header and payload were verified successfully (signature match), otherwise false.
+	 */
+	bool Verify() const;
+
+private:
+	FJsonWebToken(const FStringView InEncodedJsonWebToken, const TSharedRef<FJsonObject>& InHeaderPtr,
+		const TSharedRef<FJsonObject>& InPayloadPtr, const TOptional<TArray<uint8>>& InSignature);
+
+	static void DumpJsonObject(const FJsonObject& InJsonObject);
+
+	static TSharedPtr<FJsonObject> FromJson(const FString& InJsonStr);
+
+	static TSharedPtr<FJsonObject> ParseEncodedJson(const FStringView InEncodedJson);
+
+public:
+	// JWT payload registered claim field names
+	static const TCHAR *const CLAIM_ISSUER;
+	static const TCHAR *const CLAIM_ISSUED_AT;
+	static const TCHAR *const CLAIM_EXPIRATION;
+	static const TCHAR *const CLAIM_SUBJECT;
+	static const TCHAR *const CLAIM_AUDIENCE;
+
+	// JWT header field names
+	static const TCHAR* const HEADER_TYPE;
+	static const TCHAR* const HEADER_KEY_ID;
+	static const TCHAR* const HEADER_ALGORITHM;
+
+	// JWT header expected values
+	static const TCHAR* const TYPE_VALUE_JWT;
+
+private:
+	/** The full encoded JWT. */
+	FString EncodedJsonWebToken;
+
+	/** The decoded and parsed header. */
+	TSharedRef<FJsonObject> Header;
+
+	/** The decoded and parsed payload. */
+	TSharedRef<FJsonObject> Payload;
+
+	/** The decoded signature. */
+	TOptional<TArray<uint8>> Signature;
+};
diff --git a/Engine/Plugins/Experimental/JWT/Source/JWT/Public/JwtGlobals.h b/Engine/Plugins/Experimental/JWT/Source/JWT/Public/JwtGlobals.h
new file mode 100644
--- /dev/null
+++ b/Engine/Plugins/Experimental/JWT/Source/JWT/Public/JwtGlobals.h
@@ -0,0 +1,7 @@
+// Copyright Epic Games, Inc. All Rights Reserved.
+
+#pragma once
+
+#include "CoreMinimal.h"
+
+JWT_API DECLARE_LOG_CATEGORY_EXTERN(LogJwt, Log, All);
diff --git a/Engine/Plugins/Experimental/JWT/Source/JWT/Public/JwtIncludes.h b/Engine/Plugins/Experimental/JWT/Source/JWT/Public/JwtIncludes.h
new file mode 100644
--- /dev/null
+++ b/Engine/Plugins/Experimental/JWT/Source/JWT/Public/JwtIncludes.h
@@ -0,0 +1,14 @@
+// Copyright Epic Games, Inc. All Rights Reserved.
+
+#pragma once
+
+/* Public dependencies
+ *****************************************************************************/
+
+#include "CoreMinimal.h"
+
+ /* Public includes
+  *****************************************************************************/
+
+#include "JwtGlobals.h"
+#include "JsonWebToken.h"
