Removed 'Engine' dependency in JWT.build.cs, which was causing OSSTest target compilation to fail.

diff --git a/Engine/Plugins/Experimental/JWT/Source/JWT/JWT.Build.cs b/Engine/Plugins/Experimental/JWT/Source/JWT/JWT.Build.cs
--- a/Engine/Plugins/Experimental/JWT/Source/JWT/JWT.Build.cs
+++ b/Engine/Plugins/Experimental/JWT/Source/JWT/JWT.Build.cs
@@ -23,7 +23,6 @@ namespace UnrealBuildTool.Rules
 				{
 					"Core",
 					// ... add other public dependencies that you statically link with here ...
-					"Engine",
 					"Json",
 				}
 			);
