Fix for build blocker on Switch.  Changed `FEncryptionContextOpenSSL` to `FEncryptionContext` for platform-neutrality.  Fixed FEncryptionContext APIs for Switch and BCrypt to match OpenSSL.

diff --git a/Engine/Plugins/Experimental/JWT/Source/JWT/Private/JsonWebToken.cpp b/Engine/Plugins/Experimental/JWT/Source/JWT/Private/JsonWebToken.cpp
--- a/Engine/Plugins/Experimental/JWT/Source/JWT/Private/JsonWebToken.cpp
+++ b/Engine/Plugins/Experimental/JWT/Source/JWT/Private/JsonWebToken.cpp
@@ -343,7 +343,7 @@ bool IsSupportedAlgorithm(const FString& Algorithm)
 	return (ALGORITHM_VALUE_RS256 == Algorithm) || (ALGORITHM_VALUE_PS256 == Algorithm);
 }
 
-bool DigestVerify(FEncryptionContextOpenSSL& EncryptionContext, const FString& Algorithm, const TArrayView<const uint8> Message, TArray<uint8>& Signature, FRSAKeyHandle Key)
+bool DigestVerify(FEncryptionContext& EncryptionContext, const FString& Algorithm, const TArrayView<const uint8> Message, TArray<uint8>& Signature, FRSAKeyHandle Key)
 {
 	if (ALGORITHM_VALUE_RS256 == Algorithm)
 	{
