[JWT] Fixing JWT verify method by returning false

Fixing the JsonWebToken::Verify() method to log an error and return false instead of returning true as JWT signature verification is not implemented yet.

diff --git a/Engine/Plugins/Experimental/JWT/Source/JWT/Private/JsonWebToken.cpp b/Engine/Plugins/Experimental/JWT/Source/JWT/Private/JsonWebToken.cpp
--- a/Engine/Plugins/Experimental/JWT/Source/JWT/Private/JsonWebToken.cpp
+++ b/Engine/Plugins/Experimental/JWT/Source/JWT/Private/JsonWebToken.cpp
@@ -308,9 +308,13 @@ bool FJsonWebToken::GetClaim(const FStringView InName, TSharedPtr<FJsonValue>& O
 }
 
 
+/**
+ * Verify the JWT's signature.
+ *
+ * TODO(vri): Implement JWT verification with public key
+ */
 bool FJsonWebToken::Verify() const
 {
-	// Signature verification is yet to be implemented, so return true for now.
-	// TODO: Add support for RSA crypto algorithm (PlatformCrypto plugin) and key management so that signature verification can be performed.
-	return true;
+	UE_LOG(LogJwt, Error, TEXT("[FJsonWebToken::Verify] JWT signature verification is not implemented yet and will always return false."));
+	return false;
 }
