[JWT] Adding missing claim getters to JsonWebToken

Adding some missing claim getters to JsonWebToken: GetIssuer, GetIssuedAt, GetExpiration, GetNotBefore, GetSubject, GetAudience.

diff --git a/Engine/Plugins/Experimental/JWT/Source/JWT/Private/JsonWebToken.cpp b/Engine/Plugins/Experimental/JWT/Source/JWT/Private/JsonWebToken.cpp
--- a/Engine/Plugins/Experimental/JWT/Source/JWT/Private/JsonWebToken.cpp
+++ b/Engine/Plugins/Experimental/JWT/Source/JWT/Private/JsonWebToken.cpp
@@ -18,6 +18,7 @@ const TCHAR *const FJsonWebToken::TYPE_VALUE_JWT = TEXT("JWT");
 const TCHAR *const FJsonWebToken::CLAIM_ISSUER = TEXT("iss");
 const TCHAR *const FJsonWebToken::CLAIM_ISSUED_AT = TEXT("iat");
 const TCHAR *const FJsonWebToken::CLAIM_EXPIRATION = TEXT("exp");
+const TCHAR* const FJsonWebToken::CLAIM_NOT_BEFORE = TEXT("nbf");
 const TCHAR *const FJsonWebToken::CLAIM_SUBJECT = TEXT("sub");
 const TCHAR *const FJsonWebToken::CLAIM_AUDIENCE = TEXT("aud");
 
@@ -301,6 +302,36 @@ bool FJsonWebToken::GetAlgorithm(FString& OutValue) const
 	return Header->TryGetStringField(HEADER_ALGORITHM, OutValue);
 }
 
+bool FJsonWebToken::GetIssuer(FString& OutValue) const
+{
+	return Payload->TryGetStringField(CLAIM_ISSUER, OutValue);
+}
+
+bool FJsonWebToken::GetIssuedAt(int64& OutValue) const
+{
+	return Payload->TryGetNumberField(CLAIM_ISSUED_AT, OutValue);
+}
+
+bool FJsonWebToken::GetExpiration(int64& OutValue) const
+{
+	return Payload->TryGetNumberField(CLAIM_EXPIRATION, OutValue);
+}
+
+bool FJsonWebToken::GetNotBefore(int64& OutValue) const
+{
+	return Payload->TryGetNumberField(CLAIM_NOT_BEFORE, OutValue);
+}
+
+bool FJsonWebToken::GetSubject(FString& OutValue) const
+{
+	return Payload->TryGetStringField(CLAIM_SUBJECT, OutValue);
+}
+
+bool FJsonWebToken::GetAudience(FString& OutValue) const
+{
+	return Payload->TryGetStringField(CLAIM_AUDIENCE, OutValue);
+}
+
 bool FJsonWebToken::GetClaim(const FStringView InName, TSharedPtr<FJsonValue>& OutClaim) const
 {
 	OutClaim = Payload->TryGetField(FString(InName));
diff --git a/Engine/Plugins/Experimental/JWT/Source/JWT/Public/JsonWebToken.h b/Engine/Plugins/Experimental/JWT/Source/JWT/Public/JsonWebToken.h
--- a/Engine/Plugins/Experimental/JWT/Source/JWT/Public/JsonWebToken.h
+++ b/Engine/Plugins/Experimental/JWT/Source/JWT/Public/JsonWebToken.h
@@ -49,6 +49,54 @@ public:
 	 */
 	bool GetAlgorithm(FString& OutValue) const;
 
+	/**
+	 * Gets the issuer domain.
+	 * 
+	 * @param OutValue The value to output on success.
+	 * @return Whether the value exists and was successfully outputted.
+	 */
+	bool GetIssuer(FString& OutValue) const;
+
+	/**
+	 * Gets the issued at timestamp.
+	 *
+	 * @param OutValue The value to output on success.
+	 * @return Whether the value exists and was successfully outputted.
+	 */
+	bool GetIssuedAt(int64& OutValue) const;
+
+	/**
+	 * Gets the expiration timestamp.
+	 *
+	 * @param OutValue The value to output on success.
+	 * @return Whether the value exists and was successfully outputted.
+	 */
+	bool GetExpiration(int64& OutValue) const;
+
+	/**
+	 * Gets the not valid before timestamp.
+	 *
+	 * @param OutValue The value to output on success.
+	 * @return Whether the value exists and was successfully outputted.
+	 */
+	bool GetNotBefore(int64& OutValue) const;
+
+	/**
+	 * Gets the subject.
+	 *
+	 * @param OutValue The value to output on success.
+	 * @return Whether the value exists and was successfully outputted.
+	 */
+	bool GetSubject(FString& OutValue) const;
+
+	/**
+	 * Gets the audience.
+	 *
+	 * @param OutValue The value to output on success.
+	 * @return Whether the value exists and was successfully outputted.
+	 */
+	bool GetAudience(FString& OutValue) const;
+
 	/**
 	 * Gets a claim by name.
 	 * This method can be used to get custom claims that are not reserved as part of the JWT specification.
@@ -97,6 +145,7 @@ public:
 	static const TCHAR *const CLAIM_ISSUER;
 	static const TCHAR *const CLAIM_ISSUED_AT;
 	static const TCHAR *const CLAIM_EXPIRATION;
+	static const TCHAR* const CLAIM_NOT_BEFORE;
 	static const TCHAR *const CLAIM_SUBJECT;
 	static const TCHAR *const CLAIM_AUDIENCE;
 
