[JWT] Storing JsonWebToken's encoded header and payload parts separately

Storing the encoded header and payload parts separately to avoid having to encode them again for signature validation and to only split them once.

diff --git a/Engine/Plugins/Experimental/JWT/Source/JWT/Private/JsonWebToken.cpp b/Engine/Plugins/Experimental/JWT/Source/JWT/Private/JsonWebToken.cpp
--- a/Engine/Plugins/Experimental/JWT/Source/JWT/Private/JsonWebToken.cpp
+++ b/Engine/Plugins/Experimental/JWT/Source/JWT/Private/JsonWebToken.cpp
@@ -228,6 +228,10 @@ TOptional<FJsonWebToken> FJsonWebToken::FromString(const FStringView InEncodedJs
 		return {};
 	}
 
+	// Store the encoded header and payload for signature validation
+	FString EncodedHeaderPayload = FString::Printf(TEXT("%s.%s"),
+		*FString(EncodedHeader), *FString(EncodedPayload));
+
 	// Decode and parse the header.
 	UE_LOG(LogJwt, VeryVerbose, TEXT("[FJsonWebToken::FromString] Parsing JWT header."));
 	const TSharedPtr<FJsonObject> HeaderPtr = ParseEncodedJson(EncodedHeader);
@@ -261,7 +265,8 @@ TOptional<FJsonWebToken> FJsonWebToken::FromString(const FStringView InEncodedJs
 		Signature = MoveTemp(SignatureBytes);
 	}
 
-	const FJsonWebToken JsonWebToken(InEncodedJsonWebToken, HeaderPtr.ToSharedRef(), PayloadPtr.ToSharedRef(), Signature);
+	const FJsonWebToken JsonWebToken(EncodedHeaderPayload,
+		HeaderPtr.ToSharedRef(), PayloadPtr.ToSharedRef(), Signature);
 
 	// Validate the type.  If it exists but is not a JWT, then fail.
 	FString Type;
@@ -274,13 +279,14 @@ TOptional<FJsonWebToken> FJsonWebToken::FromString(const FStringView InEncodedJs
 	return JsonWebToken;
 }
 
+
 FJsonWebToken::FJsonWebToken(
-	const FStringView InEncodedJsonWebToken,
+	const FStringView InEncodedHeaderPayload,
 	const TSharedRef<FJsonObject>& InHeaderPtr,
 	const TSharedRef<FJsonObject>& InPayloadPtr,
 	const TOptional<TArray<uint8>>& InSignature
 )
-	: EncodedJsonWebToken(InEncodedJsonWebToken)
+	: EncodedHeaderPayload(InEncodedHeaderPayload)
 	, Header(InHeaderPtr)
 	, Payload(InPayloadPtr)
 	, Signature(InSignature)
diff --git a/Engine/Plugins/Experimental/JWT/Source/JWT/Public/JsonWebToken.h b/Engine/Plugins/Experimental/JWT/Source/JWT/Public/JsonWebToken.h
--- a/Engine/Plugins/Experimental/JWT/Source/JWT/Public/JsonWebToken.h
+++ b/Engine/Plugins/Experimental/JWT/Source/JWT/Public/JsonWebToken.h
@@ -131,7 +131,8 @@ public:
 	bool Verify() const;
 
 private:
-	FJsonWebToken(const FStringView InEncodedJsonWebToken, const TSharedRef<FJsonObject>& InHeaderPtr,
+	FJsonWebToken(
+		const FStringView InEncodedHeaderPayload, const TSharedRef<FJsonObject>& InHeaderPtr,
 		const TSharedRef<FJsonObject>& InPayloadPtr, const TOptional<TArray<uint8>>& InSignature);
 
 	static void DumpJsonObject(const FJsonObject& InJsonObject);
@@ -158,8 +159,9 @@ public:
 	static const TCHAR* const TYPE_VALUE_JWT;
 
 private:
-	/** The full encoded JWT. */
-	FString EncodedJsonWebToken;
+
+	/** The encoded header and payload parts. */
+	FString EncodedHeaderPayload;
 
 	/** The decoded and parsed header. */
 	TSharedRef<FJsonObject> Header;
