[JWT] Fixing base64 decode mode in JsonWebToken utility methods

Fixing the base64 decode mode to be url-safe to not fail with special characters in JsonWebToken utility methods.

diff --git a/Engine/Plugins/Experimental/JWT/Source/JWT/Private/JsonWebToken.cpp b/Engine/Plugins/Experimental/JWT/Source/JWT/Private/JsonWebToken.cpp
--- a/Engine/Plugins/Experimental/JWT/Source/JWT/Private/JsonWebToken.cpp
+++ b/Engine/Plugins/Experimental/JWT/Source/JWT/Private/JsonWebToken.cpp
@@ -69,31 +69,31 @@ namespace
 		return true;
 	}
 
-	bool Base64Decode(const FStringView InSource, FString& OutDest)
+	bool Base64UrlDecode(const FStringView InSource, FString& OutDest)
 	{
 		if (InSource.IsEmpty())
 		{
 			return false;
 		}
 
-		return FBase64::Decode(FString(InSource), OutDest);
+		return FBase64::Decode(FString(InSource), OutDest, EBase64Mode::UrlSafe);
 	}
 
-	bool Base64Decode(const FStringView InSource, TArray<uint8>& OutDest)
+	bool Base64UrlDecode(const FStringView InSource, TArray<uint8>& OutDest)
 	{
 		if (InSource.IsEmpty())
 		{
 			return false;
 		}
 
-		return FBase64::Decode(FString(InSource), OutDest);
+		return FBase64::Decode(FString(InSource), OutDest, EBase64Mode::UrlSafe);
 	}
 
 	bool StringViewToBytes(const FStringView In, TArray<uint8>& OutBytes, const bool IsEncoded)
 	{
 		if (IsEncoded)
 		{
-			return Base64Decode(In, OutBytes);
+			return Base64UrlDecode(In, OutBytes);
 		}
 
 		OutBytes.Reserve(In.Len());
@@ -190,7 +190,7 @@ TSharedPtr<FJsonObject> FJsonWebToken::ParseEncodedJson(const FStringView InEnco
 	TSharedPtr<FJsonObject> ParsedObj;
 
 	FString DecodedJson;
-	if (Base64Decode(InEncodedJson, DecodedJson))
+	if (Base64UrlDecode(InEncodedJson, DecodedJson))
 	{
 		ParsedObj = FromJson(DecodedJson);
 	}
