[JWT] Removing unnecessary bIsSignatureEncoded argument from JsonWebToken

Removing the bIsSignatureEncoded argument from JsonWebToken factory and utility methods because the spec does not allow for a un-encoded signature. Replacing with a decode by default.

diff --git a/Engine/Plugins/Experimental/JWT/Source/JWT/Private/IJwt.cpp b/Engine/Plugins/Experimental/JWT/Source/JWT/Private/IJwt.cpp
--- a/Engine/Plugins/Experimental/JWT/Source/JWT/Private/IJwt.cpp
+++ b/Engine/Plugins/Experimental/JWT/Source/JWT/Private/IJwt.cpp
@@ -32,14 +32,14 @@ void FJwt::ShutdownModule()
 
 namespace UE::JWT
 {
-	TOptional<FJsonWebToken> FromString(const FStringView InEncodedJsonWebToken, const bool bIsSignatureEncoded)
+	TOptional<FJsonWebToken> FromString(const FStringView InEncodedJsonWebToken)
 	{
-		return FJsonWebToken::FromString(InEncodedJsonWebToken, bIsSignatureEncoded);
+		return FJsonWebToken::FromString(InEncodedJsonWebToken);
 	}
 
 
-	bool FromString(const FStringView InEncodedJsonWebToken, FJsonWebToken& OutJsonWebToken, const bool bIsSignatureEncoded)
+	bool FromString(const FStringView InEncodedJsonWebToken, FJsonWebToken& OutJsonWebToken)
 	{
-		return FJsonWebToken::FromString(InEncodedJsonWebToken, OutJsonWebToken, bIsSignatureEncoded);
+		return FJsonWebToken::FromString(InEncodedJsonWebToken, OutJsonWebToken);
 	}
 }
diff --git a/Engine/Plugins/Experimental/JWT/Source/JWT/Private/JsonWebToken.cpp b/Engine/Plugins/Experimental/JWT/Source/JWT/Private/JsonWebToken.cpp
--- a/Engine/Plugins/Experimental/JWT/Source/JWT/Private/JsonWebToken.cpp
+++ b/Engine/Plugins/Experimental/JWT/Source/JWT/Private/JsonWebToken.cpp
@@ -198,9 +198,10 @@ TSharedPtr<FJsonObject> FJsonWebToken::ParseEncodedJson(const FStringView InEnco
 	return ParsedObj;
 }
 
-bool FJsonWebToken::FromString(const FStringView InEncodedJsonWebToken, FJsonWebToken& OutJsonWebToken, const bool bIsSignatureEncoded)
+bool FJsonWebToken::FromString(const FStringView InEncodedJsonWebToken, FJsonWebToken& OutJsonWebToken)
 {
-	TOptional<FJsonWebToken> JsonWebToken = FromString(InEncodedJsonWebToken, bIsSignatureEncoded);
+	TOptional<FJsonWebToken> JsonWebToken = FromString(InEncodedJsonWebToken);
+
 	if (!JsonWebToken.IsSet())
 	{
 		return false;
@@ -210,7 +211,7 @@ bool FJsonWebToken::FromString(const FStringView InEncodedJsonWebToken, FJsonWeb
 	return true;
 }
 
-TOptional<FJsonWebToken> FJsonWebToken::FromString(const FStringView InEncodedJsonWebToken, const bool bIsSignatureEncoded)
+TOptional<FJsonWebToken> FJsonWebToken::FromString(const FStringView InEncodedJsonWebToken)
 {
 	// Check for the correct number of dots.
 	if (Algo::Count(InEncodedJsonWebToken, TEXT('.')) != 2)
@@ -252,11 +253,14 @@ TOptional<FJsonWebToken> FJsonWebToken::FromString(const FStringView InEncodedJs
 
 	// Decode (but do not parse) the signature if it not empty.
 	TOptional<TArray<uint8>> Signature;
+
 	if (!SignaturePart.IsEmpty())
 	{
 		UE_LOG(LogJwt, VeryVerbose, TEXT("[FJsonWebToken::FromString] Decoding JWT signature."));
+
 		TArray<uint8> SignatureBytes;
-		if (!StringViewToBytes(SignaturePart, SignatureBytes, bIsSignatureEncoded))
+
+		if (!Base64UrlDecode(SignaturePart, SignatureBytes))
 		{
 			UE_LOG(LogJwt, Verbose, TEXT("[FJsonWebToken::FromString] Failed to decode the signature."));
 			return {};
diff --git a/Engine/Plugins/Experimental/JWT/Source/JWT/Public/JsonWebToken.h b/Engine/Plugins/Experimental/JWT/Source/JWT/Public/JsonWebToken.h
--- a/Engine/Plugins/Experimental/JWT/Source/JWT/Public/JsonWebToken.h
+++ b/Engine/Plugins/Experimental/JWT/Source/JWT/Public/JsonWebToken.h
@@ -8,6 +8,7 @@ class FJsonObject;
 class FJsonValue;
 enum class EJson;
 
+
 class JWT_API FJsonWebToken
 {
 public:
@@ -21,9 +22,9 @@ public:
 	 * @param InJsonWebTokenString The string to decode.
 	 * @return An optional set with a FJsonWebToken if the JWT was successfully decoded.
 	 */
-	static TOptional<FJsonWebToken> FromString(const FStringView InEncodedJsonWebToken, const bool bIsSignatureEncoded);
+	static TOptional<FJsonWebToken> FromString(const FStringView InEncodedJsonWebToken);
 
-	static bool FromString(const FStringView InEncodedJsonWebToken, FJsonWebToken& OutJsonWebToken, const bool bIsSignatureEncoded);
+	static bool FromString(const FStringView InEncodedJsonWebToken, FJsonWebToken& OutJsonWebToken);
 
 	/**
 	 * Gets the type.
