[JWT] Adapting header of Jwt module interface as well to remove bIsSignatureEncoded argument

diff --git a/Engine/Plugins/Experimental/JWT/Source/JWT/Public/IJwt.h b/Engine/Plugins/Experimental/JWT/Source/JWT/Public/IJwt.h
--- a/Engine/Plugins/Experimental/JWT/Source/JWT/Public/IJwt.h
+++ b/Engine/Plugins/Experimental/JWT/Source/JWT/Public/IJwt.h
@@ -46,9 +46,9 @@ public:
 
 namespace UE::JWT
 {
-	JWT_API TOptional<FJsonWebToken> FromString(const FStringView InEncodedJsonWebToken, const bool bIsSignatureEncoded);
+	JWT_API TOptional<FJsonWebToken> FromString(const FStringView InEncodedJsonWebToken);
 
-	JWT_API bool FromString(const FStringView InEncodedJWT, FJsonWebToken& OutJsonWebToken, const bool bIsSignatureEncoded);
+	JWT_API bool FromString(const FStringView InEncodedJWT, FJsonWebToken& OutJsonWebToken);
 }
 
 #if UE_ENABLE_INCLUDE_ORDER_DEPRECATED_IN_5_2
