[JWT] Implementing the JWT algorithm interface with RS256

Implementing the JWT algorithm interface for customizable JWT verification and signature validation, including the first supported algorithm RS256.

diff --git a/Engine/Plugins/Experimental/JWT/Source/JWT/JWT.Build.cs b/Engine/Plugins/Experimental/JWT/Source/JWT/JWT.Build.cs
--- a/Engine/Plugins/Experimental/JWT/Source/JWT/JWT.Build.cs
+++ b/Engine/Plugins/Experimental/JWT/Source/JWT/JWT.Build.cs
@@ -6,38 +6,19 @@ namespace UnrealBuildTool.Rules
 	{
 		public JWT(ReadOnlyTargetRules Target) : base(Target)
 		{
-			PublicIncludePaths.AddRange(
-				new string[] {
-					// ... add public include paths required here ...
-				}
-			);
-
-			PrivateIncludePaths.AddRange(
-				new string[] {
-					// ... add other private include paths required here ...
-				}
-			);
-
 			PublicDependencyModuleNames.AddRange(
 				new string[]
 				{
-					"Core",
-					// ... add other public dependencies that you statically link with here ...
-					"Json",
+					"Core"
 				}
 			);
 
 			PrivateDependencyModuleNames.AddRange(
 				new string[]
 				{
-					// ... add private dependencies that you statically link with here ...
-				}
-			);
-
-			DynamicallyLoadedModuleNames.AddRange(
-				new string[]
-				{
-					// ... add any modules that your module loads dynamically here ...
+					"Json",
+					"PlatformCrypto",
+					"PlatformCryptoOpenSSL"
 				}
 			);
 		}
diff --git a/Engine/Plugins/Experimental/JWT/Source/JWT/Private/JwtAlgorithms.cpp b/Engine/Plugins/Experimental/JWT/Source/JWT/Private/JwtAlgorithms.cpp
new file mode 100644
--- /dev/null
+++ b/Engine/Plugins/Experimental/JWT/Source/JWT/Private/JwtAlgorithms.cpp
@@ -0,0 +1,83 @@
+// Copyright Epic Games, Inc. All Rights Reserved.
+
+#include "JwtAlgorithms.h"
+
+#include "JwtUtils.h"
+
+
+FJwtAlgorithm_RS256::~FJwtAlgorithm_RS256()
+{
+	if (PublicKey)
+	{
+		DestroyKey(PublicKey);
+	}
+}
+
+
+bool FJwtAlgorithm_RS256::VerifySignature(
+	const TArrayView<const uint8> EncodedMessage,
+	const TArrayView<const uint8> DecodedSignature) const
+{
+	TUniquePtr<FEncryptionContext> EncryptionContext
+		= FJwtUtils::GetEncryptionContext();
+
+	if (!EncryptionContext.IsValid())
+	{
+		return false;
+	}
+
+	TArray<uint8> HashedMessage;
+
+	// Hash the encoded message
+	if (!EncryptionContext->CalcSHA256(EncodedMessage, HashedMessage))
+	{
+		return false;
+	}
+
+	return EncryptionContext->DigestVerify_RS256(
+		HashedMessage, DecodedSignature, PublicKey);
+}
+
+
+bool FJwtAlgorithm_RS256::SetPublicKey(const TArrayView<const uint8> InKey)
+{
+	TUniquePtr<FEncryptionContext> EncryptionContext
+		= FJwtUtils::GetEncryptionContext();
+
+	if (!EncryptionContext.IsValid())
+	{
+		return false;
+	}
+
+	if (PublicKey)
+	{
+		DestroyKey(PublicKey);
+	}
+
+	PublicKey = EncryptionContext->GetPublicKey_RSA(InKey);
+
+	if (!PublicKey)
+	{
+		UE_LOG(LogJwt, Error,
+			TEXT("[FJwtAlgorithm_RS256::SetPublicKey] RSA public key is invalid."));
+
+		return false;
+	}
+
+	return true;
+}
+
+
+void FJwtAlgorithm_RS256::DestroyKey(void* Key)
+{
+	TUniquePtr<FEncryptionContext> EncryptionContext
+		= FJwtUtils::GetEncryptionContext();
+
+	if (!EncryptionContext.IsValid())
+	{
+		return;
+	}
+
+	EncryptionContext->DestroyKey_RSA(Key);
+}
+
diff --git a/Engine/Plugins/Experimental/JWT/Source/JWT/Private/JwtUtils.cpp b/Engine/Plugins/Experimental/JWT/Source/JWT/Private/JwtUtils.cpp
new file mode 100644
--- /dev/null
+++ b/Engine/Plugins/Experimental/JWT/Source/JWT/Private/JwtUtils.cpp
@@ -0,0 +1,20 @@
+// Copyright Epic Games, Inc. All Rights Reserved.
+
+#include "JwtUtils.h"
+#include "JwtGlobals.h"
+
+
+TUniquePtr<FEncryptionContext> FJwtUtils::GetEncryptionContext()
+{
+	TUniquePtr<FEncryptionContext> EncryptionContext
+		= IPlatformCrypto::Get().CreateContext();
+
+	if (!EncryptionContext.IsValid())
+	{
+		UE_LOG(LogJwt, Error,
+			TEXT("[JwtUtils::GetEncryptionContext] "
+				"EncryptionContext pointer is invalid."));
+	}
+
+	return EncryptionContext;
+}
diff --git a/Engine/Plugins/Experimental/JWT/Source/JWT/Public/JwtAlgorithm.h b/Engine/Plugins/Experimental/JWT/Source/JWT/Public/JwtAlgorithm.h
new file mode 100644
--- /dev/null
+++ b/Engine/Plugins/Experimental/JWT/Source/JWT/Public/JwtAlgorithm.h
@@ -0,0 +1,38 @@
+// Copyright Epic Games, Inc. All Rights Reserved.
+
+#pragma once
+
+#include "JwtGlobals.h"
+
+#include "CoreMinimal.h"
+
+
+class IJwtAlgorithm
+{
+
+public:
+
+	virtual ~IJwtAlgorithm()
+	{
+	}
+
+public:
+
+	/**
+	 * Gets the algorithm's string representation.
+	 */
+	virtual inline const FString& GetAlgString() = 0;
+
+	/**
+	 * Verifies the decoded signature with the encoded message.
+	 * 
+	 * @param EncodedMessage The base64 encoded message {header}.{payload}
+	 * @param DecodedSignature The decoded JWT signature
+	 * 
+	 * @return Whether the signature was successfully verified.
+	 */
+	virtual bool VerifySignature(
+		const TArrayView<const uint8> EncodedMessage,
+		const TArrayView<const uint8> DecodedSignature) const = 0;
+
+};
diff --git a/Engine/Plugins/Experimental/JWT/Source/JWT/Public/JwtAlgorithms.h b/Engine/Plugins/Experimental/JWT/Source/JWT/Public/JwtAlgorithms.h
new file mode 100644
--- /dev/null
+++ b/Engine/Plugins/Experimental/JWT/Source/JWT/Public/JwtAlgorithms.h
@@ -0,0 +1,68 @@
+// Copyright Epic Games, Inc. All Rights Reserved.
+
+#pragma once
+
+#include "JwtAlgorithm.h"
+
+
+typedef void* FRSAKeyHandle;
+
+
+class JWT_API FJwtAlgorithm_RS256
+	: public IJwtAlgorithm
+{
+
+public:
+
+	FJwtAlgorithm_RS256() = default;
+
+	FJwtAlgorithm_RS256(const FJwtAlgorithm_RS256&) = delete;
+
+	FJwtAlgorithm_RS256& operator=(const FJwtAlgorithm_RS256&) = delete;
+
+	~FJwtAlgorithm_RS256();
+
+public:
+
+	// ~ IJwtAlgorithm interface
+
+	virtual inline const FString& GetAlgString() override
+	{
+		return AlgorithmString;
+	}
+
+	virtual bool VerifySignature(
+		const TArrayView<const uint8> EncodedMessage,
+		const TArrayView<const uint8> DecodedSignature) const override;
+
+public:
+
+	/*
+	 * Set the public RSA key.
+	 *
+	 * @param InKey Raw key bytes
+	 *
+	 * @return Whether the operation was successful
+	 */
+	bool SetPublicKey(const TArrayView<const uint8> InKey);
+
+private:
+
+	/**
+	 * Free the RSA key.
+	 *
+	 * @param Key The key pointer
+	 */
+	void DestroyKey(void* Key);
+
+protected:
+
+	/** Holds the algorithm string. */
+	const FString AlgorithmString = "RS256";
+
+private:
+
+	/** Holds the public key pointer. */
+	void* PublicKey = nullptr;
+
+};
diff --git a/Engine/Plugins/Experimental/JWT/Source/JWT/Public/JwtUtils.h b/Engine/Plugins/Experimental/JWT/Source/JWT/Public/JwtUtils.h
new file mode 100644
--- /dev/null
+++ b/Engine/Plugins/Experimental/JWT/Source/JWT/Public/JwtUtils.h
@@ -0,0 +1,23 @@
+// Copyright Epic Games, Inc. All Rights Reserved.
+
+#pragma once
+
+#include "CoreMinimal.h"
+
+#include "IPlatformCrypto.h"
+
+
+class FJwtUtils
+{
+
+public:
+
+	/**
+	 * Get the pointer to the current platform
+	 * encryption context (e.g. OpenSSL or SwitchSSL).
+	 *
+	 * @return Pointer to the encryption context
+	 */
+	static TUniquePtr<FEncryptionContext> GetEncryptionContext();
+
+};
