[JWT] Adding PlatformCrypto plugin dependency to JWT

diff --git a/Engine/Plugins/Experimental/JWT/JWT.uplugin b/Engine/Plugins/Experimental/JWT/JWT.uplugin
--- a/Engine/Plugins/Experimental/JWT/JWT.uplugin
+++ b/Engine/Plugins/Experimental/JWT/JWT.uplugin
@@ -13,13 +13,18 @@
 	"EnabledByDefault" : false,
 	"CanContainContent" : false,
 	"IsBetaVersion" : false,
-	"Installed" : false,
-	"Modules" :
-	[
+	"Installed": false,
+	"Modules": [
 		{
-			"Name" : "JWT",
-			"Type" : "Runtime",
-			"LoadingPhase" : "Default"
+			"Name": "JWT",
+			"Type": "Runtime",
+			"LoadingPhase": "Default"
+		}
+	],
+	"Plugins": [
+		{
+			"Name": "PlatformCrypto",
+			"Enabled": true
 		}
 	]
 }
