[JWT] Refactoring JWT utility methods

Moving, improving and reformatting the JWT utility methods so they can be used by other functions and classes within JWT.

diff --git a/Engine/Plugins/Experimental/JWT/Source/JWT/Private/JsonWebToken.cpp b/Engine/Plugins/Experimental/JWT/Source/JWT/Private/JsonWebToken.cpp
--- a/Engine/Plugins/Experimental/JWT/Source/JWT/Private/JsonWebToken.cpp
+++ b/Engine/Plugins/Experimental/JWT/Source/JWT/Private/JsonWebToken.cpp
@@ -1,11 +1,14 @@
 // Copyright Epic Games, Inc. All Rights Reserved.
 
 #include "JsonWebToken.h"
-#include "Algo/Count.h"
+
 #include "JwtGlobals.h"
-#include "Misc/Base64.h"
+#include "JwtUtils.h"
+
+#include "Algo/Count.h"
 #include "Serialization/JsonSerializer.h"
 
+
 // JWT header field names
 const TCHAR *const FJsonWebToken::HEADER_TYPE = TEXT("typ");
 const TCHAR *const FJsonWebToken::HEADER_KEY_ID = TEXT("kid");
@@ -23,89 +26,6 @@ const TCHAR *const FJsonWebToken::CLAIM_SUBJECT = TEXT("sub");
 const TCHAR *const FJsonWebToken::CLAIM_AUDIENCE = TEXT("aud");
 
 
-namespace
-{
-	/**
-	 * Helper for splitting FStringViews, like FString.
-	 */
-	bool SplitStringView(const FStringView InSourceString, const TCHAR InDelimiter, FStringView& OutLeft, FStringView& OutRight)
-	{
-		int32 Position = INDEX_NONE;
-		if (!InSourceString.FindChar(InDelimiter, Position) || Position == INDEX_NONE)
-		{
-			return false;
-		}
-
-		OutLeft = InSourceString.Left(Position);
-		OutRight = InSourceString.RightChop(Position + 1);
-		return true;
-	}
-
-	/**
-	 * Splits an encoded JWT string into the 3 header, payload, and signature parts.
-	 */
-	bool SplitEncodedJsonWebTokenString(const FStringView InEncodedJsonWebTokenString, FStringView& OutEncodedHeaderPart, FStringView& OutEncodedPayloadPart, FStringView& OutEncodedSignaturePart)
-	{
-		FStringView Rest;
-
-		if (!SplitStringView(InEncodedJsonWebTokenString, TEXT('.'), OutEncodedHeaderPart, Rest))
-		{
-			UE_LOG(LogJwt, Warning, TEXT("[SplitEncodedJsonWebTokenString] Cannot extract header from token string."));
-			return false;
-		}
-
-		if (!SplitStringView(Rest, TEXT('.'), OutEncodedPayloadPart, OutEncodedSignaturePart))
-		{
-			UE_LOG(LogJwt, Warning, TEXT("[SplitEncodedJsonWebTokenString] Cannot extract payload from token string."));
-			return false;
-		}
-
-		if (OutEncodedHeaderPart.IsEmpty() || OutEncodedPayloadPart.IsEmpty())
-		{
-			UE_LOG(LogJwt, Warning, TEXT("[SplitEncodedJsonWebTokenString] Empty header and/or payload in token string."));
-			return false;
-		}
-
-		return true;
-	}
-
-	bool Base64UrlDecode(const FStringView InSource, FString& OutDest)
-	{
-		if (InSource.IsEmpty())
-		{
-			return false;
-		}
-
-		return FBase64::Decode(FString(InSource), OutDest, EBase64Mode::UrlSafe);
-	}
-
-	bool Base64UrlDecode(const FStringView InSource, TArray<uint8>& OutDest)
-	{
-		if (InSource.IsEmpty())
-		{
-			return false;
-		}
-
-		return FBase64::Decode(FString(InSource), OutDest, EBase64Mode::UrlSafe);
-	}
-
-	bool StringViewToBytes(const FStringView In, TArray<uint8>& OutBytes, const bool IsEncoded)
-	{
-		if (IsEncoded)
-		{
-			return Base64UrlDecode(In, OutBytes);
-		}
-
-		OutBytes.Reserve(In.Len());
-		for (const TCHAR& Ch : In)
-		{
-			OutBytes.Add(static_cast<uint8>(Ch));
-		}
-
-		return true;
-	}
-}
-
 void FJsonWebToken::DumpJsonObject(const FJsonObject& InJsonObject)
 {
 	if (!UE_LOG_ACTIVE(LogJwt, VeryVerbose))
@@ -190,7 +110,7 @@ TSharedPtr<FJsonObject> FJsonWebToken::ParseEncodedJson(const FStringView InEnco
 	TSharedPtr<FJsonObject> ParsedObj;
 
 	FString DecodedJson;
-	if (Base64UrlDecode(InEncodedJson, DecodedJson))
+	if (FJwtUtils::Base64UrlDecode(InEncodedJson, DecodedJson))
 	{
 		ParsedObj = FromJson(DecodedJson);
 	}
@@ -224,7 +144,7 @@ TOptional<FJsonWebToken> FJsonWebToken::FromString(const FStringView InEncodedJs
 	FStringView EncodedHeader;
 	FStringView EncodedPayload;
 	FStringView SignaturePart;
-	if (!SplitEncodedJsonWebTokenString(InEncodedJsonWebToken, EncodedHeader, EncodedPayload, SignaturePart))
+	if (!FJwtUtils::SplitEncodedJsonWebTokenString(InEncodedJsonWebToken, EncodedHeader, EncodedPayload, SignaturePart))
 	{
 		return {};
 	}
@@ -260,7 +180,7 @@ TOptional<FJsonWebToken> FJsonWebToken::FromString(const FStringView InEncodedJs
 
 		TArray<uint8> SignatureBytes;
 
-		if (!Base64UrlDecode(SignaturePart, SignatureBytes))
+		if (!FJwtUtils::Base64UrlDecode(SignaturePart, SignatureBytes))
 		{
 			UE_LOG(LogJwt, Verbose, TEXT("[FJsonWebToken::FromString] Failed to decode the signature."));
 			return {};
diff --git a/Engine/Plugins/Experimental/JWT/Source/JWT/Private/JwtAlgorithms.cpp b/Engine/Plugins/Experimental/JWT/Source/JWT/Private/JwtAlgorithms.cpp
--- a/Engine/Plugins/Experimental/JWT/Source/JWT/Private/JwtAlgorithms.cpp
+++ b/Engine/Plugins/Experimental/JWT/Source/JWT/Private/JwtAlgorithms.cpp
@@ -68,6 +68,16 @@ bool FJwtAlgorithm_RS256::SetPublicKey(const TArrayView<const uint8> InKey)
 }
 
 
+bool FJwtAlgorithm_RS256::SetPublicKey(const FStringView InKey)
+{
+	TArray<uint8> PublicKeyBytes;
+
+	FJwtUtils::StringViewToBytes(InKey, PublicKeyBytes);
+
+	return SetPublicKey(PublicKeyBytes);
+}
+
+
 void FJwtAlgorithm_RS256::DestroyKey(void* Key)
 {
 	TUniquePtr<FEncryptionContext> EncryptionContext
diff --git a/Engine/Plugins/Experimental/JWT/Source/JWT/Private/JwtUtils.cpp b/Engine/Plugins/Experimental/JWT/Source/JWT/Private/JwtUtils.cpp
--- a/Engine/Plugins/Experimental/JWT/Source/JWT/Private/JwtUtils.cpp
+++ b/Engine/Plugins/Experimental/JWT/Source/JWT/Private/JwtUtils.cpp
@@ -3,6 +3,8 @@
 #include "JwtUtils.h"
 #include "JwtGlobals.h"
 
+#include "Misc/Base64.h"
+
 
 TUniquePtr<FEncryptionContext> FJwtUtils::GetEncryptionContext()
 {
@@ -18,3 +20,96 @@ TUniquePtr<FEncryptionContext> FJwtUtils::GetEncryptionContext()
 
 	return EncryptionContext;
 }
+
+
+bool FJwtUtils::SplitStringView(
+	const FStringView InSource, const TCHAR InDelimiter,
+	FStringView& OutLeft, FStringView& OutRight)
+{
+	int32 Position = INDEX_NONE;
+
+	if (!InSource.FindChar(InDelimiter, Position) || Position == INDEX_NONE)
+	{
+		UE_LOG(LogJwt, Error,
+			TEXT("[JwtUtils::SplitStringView] "
+				"Could not find delimiter in source string."));
+
+		return false;
+	}
+
+	OutLeft = InSource.Left(Position);
+	OutRight = InSource.RightChop(Position + 1);
+
+	return true;
+}
+
+
+bool FJwtUtils::SplitEncodedJsonWebTokenString(
+	const FStringView InEncodedJsonWebTokenString, FStringView& OutEncodedHeaderPart,
+	FStringView& OutEncodedPayloadPart, FStringView& OutEncodedSignaturePart)
+{
+	FStringView Rest;
+
+	if (!SplitStringView(
+		InEncodedJsonWebTokenString, TEXT('.'), OutEncodedHeaderPart, Rest))
+	{
+		UE_LOG(LogJwt, Error,
+			TEXT("[JwtUtils::SplitEncodedJsonWebTokenString] "
+				"Cannot extract header from token string."));
+
+		return false;
+	}
+
+	if (!SplitStringView(Rest, TEXT('.'), OutEncodedPayloadPart, OutEncodedSignaturePart))
+	{
+		UE_LOG(LogJwt, Error,
+			TEXT("[JwtUtils::SplitEncodedJsonWebTokenString] "
+				"Cannot extract payload from token string."));
+
+		return false;
+	}
+
+	if (OutEncodedHeaderPart.IsEmpty() || OutEncodedPayloadPart.IsEmpty())
+	{
+		UE_LOG(LogJwt, Error,
+			TEXT("[JwtUtils::SplitEncodedJsonWebTokenString] "
+				"Empty header and/or payload in token string."));
+
+		return false;
+	}
+
+	return true;
+}
+
+
+bool FJwtUtils::Base64UrlDecode(const FStringView InSource, FString& OutDest)
+{
+	if (InSource.IsEmpty())
+	{
+		return false;
+	}
+
+	return FBase64::Decode(FString(InSource), OutDest, EBase64Mode::UrlSafe);
+}
+
+
+bool FJwtUtils::Base64UrlDecode(const FStringView InSource, TArray<uint8>& OutDest)
+{
+	if (InSource.IsEmpty())
+	{
+		return false;
+	}
+
+	return FBase64::Decode(FString(InSource), OutDest, EBase64Mode::UrlSafe);
+}
+
+
+void FJwtUtils::StringViewToBytes(const FStringView InSource, TArray<uint8>& OutBytes)
+{
+	OutBytes.Reserve(InSource.Len());
+
+	for (const TCHAR& Ch : InSource)
+	{
+		OutBytes.Add(static_cast<uint8>(Ch));
+	}
+}
diff --git a/Engine/Plugins/Experimental/JWT/Source/JWT/Public/JwtAlgorithms.h b/Engine/Plugins/Experimental/JWT/Source/JWT/Public/JwtAlgorithms.h
--- a/Engine/Plugins/Experimental/JWT/Source/JWT/Public/JwtAlgorithms.h
+++ b/Engine/Plugins/Experimental/JWT/Source/JWT/Public/JwtAlgorithms.h
@@ -5,9 +5,6 @@
 #include "JwtAlgorithm.h"
 
 
-typedef void* FRSAKeyHandle;
-
-
 class JWT_API FJwtAlgorithm_RS256
 	: public IJwtAlgorithm
 {
@@ -46,6 +43,15 @@ public:
 	 */
 	bool SetPublicKey(const TArrayView<const uint8> InKey);
 
+	/**
+	 * Set the public RSA key.
+	 *
+	 * @param InKey Key as string view
+	 *
+	 * @return Whether the operation was successful
+	 */
+	bool SetPublicKey(const FStringView InKey);
+
 private:
 
 	/**
diff --git a/Engine/Plugins/Experimental/JWT/Source/JWT/Public/JwtUtils.h b/Engine/Plugins/Experimental/JWT/Source/JWT/Public/JwtUtils.h
--- a/Engine/Plugins/Experimental/JWT/Source/JWT/Public/JwtUtils.h
+++ b/Engine/Plugins/Experimental/JWT/Source/JWT/Public/JwtUtils.h
@@ -20,4 +20,60 @@ public:
 	 */
 	static TUniquePtr<FEncryptionContext> GetEncryptionContext();
 
+	/**
+	 * Split a StringView on a delimiter.
+	 *
+	 * @param InSource The source string view
+	 * @param InDelimiter The delimiter
+	 * @param OutLeft Out string view for the left part
+	 * @param OutRight Out string view for the right part
+	 *
+	 * @return Whether the string view could be split
+	 */
+	static bool SplitStringView(
+		const FStringView InSource, const TCHAR InDelimiter,
+		FStringView& OutLeft, FStringView& OutRight);
+
+	/**
+	 * Split an encoded JWT string into the 3 header, payload, and signature parts.
+	 *
+	 * @param InEncodedJsonWebTokenString The JWT string
+	 * @param OutEncodedHeaderPart Out string for the header part
+	 * @param OutEncodedPayloadPart Out string for the payload part
+	 * @param OutEncodedSignaturePart Out string for the signature part
+	 *
+	 * @return Whether the JWT string could be split
+	 */
+	static bool SplitEncodedJsonWebTokenString(
+		const FStringView InEncodedJsonWebTokenString, FStringView& OutEncodedHeaderPart,
+		FStringView& OutEncodedPayloadPart, FStringView& OutEncodedSignaturePart);
+
+	/**
+	 * Decode base64 in url-safe mode from string to string.
+	 *
+	 * @param InSource The source string
+	 * @param OutDest The destination string
+	 *
+	 * @return Whether the string could be decoded
+	 */
+	static bool Base64UrlDecode(const FStringView InSource, FString& OutDest);
+
+	/**
+	 * Decode base64 in url-safe mode from string to byte array.
+	 *
+	 * @param InSource The source string
+	 * @param OutDest The destination byte array
+	 *
+	 * @return Whether the string could be decoded
+	 */
+	static bool Base64UrlDecode(const FStringView InSource, TArray<uint8>& OutDest);
+
+	/**
+	 * Convert a StringView to a byte array.
+	 *
+	 * @param InSource The source string
+	 * @param OutBytes The destination byte array
+	 */
+	static void StringViewToBytes(const FStringView InSource, TArray<uint8>& OutBytes);
+
 };
