[JWT] Adding SupportedTargetPlatforms and PlatformAllowList to JWT

diff --git a/Engine/Plugins/Experimental/JWT/JWT.uplugin b/Engine/Plugins/Experimental/JWT/JWT.uplugin
--- a/Engine/Plugins/Experimental/JWT/JWT.uplugin
+++ b/Engine/Plugins/Experimental/JWT/JWT.uplugin
@@ -1,24 +1,35 @@
 {
-	"FileVersion" : 3,
-	"Version" : 1,
-	"VersionName" : "1.0",
-	"FriendlyName" : "JSON Web Token Plugin",
-	"Description" : "An API for working with JSON Web Token (JWT) data.",
-	"Category" : "Misc",
-	"CreatedBy" : "Epic Games, Inc.",
-	"CreatedByURL" : "https://epicgames.com",
-	"DocsURL" : "",
-	"MarketplaceURL" : "",
-	"SupportURL" : "",
-	"EnabledByDefault" : false,
-	"CanContainContent" : false,
-	"IsBetaVersion" : false,
+	"FileVersion": 3,
+	"Version": 1,
+	"VersionName": "1.0",
+	"FriendlyName": "JSON Web Token Plugin",
+	"Description": "An API for working with JSON Web Token (JWT) data.",
+	"Category": "Misc",
+	"CreatedBy": "Epic Games, Inc.",
+	"CreatedByURL": "https://epicgames.com",
+	"DocsURL": "",
+	"MarketplaceURL": "",
+	"SupportURL": "",
+	"EnabledByDefault": false,
+	"CanContainContent": false,
+	"IsBetaVersion": false,
 	"Installed": false,
 	"Modules": [
 		{
 			"Name": "JWT",
 			"Type": "Runtime",
-			"LoadingPhase": "Default"
+			"LoadingPhase": "Default",
+			"PlatformAllowList": [
+				"Android",
+				"IOS",
+				"Mac",
+				"Win64",
+				"Linux",
+				"LinuxArm64"
+			],
+			"PlatformDenyList": [
+				"TVOS"
+			]
 		}
 	],
 	"Plugins": [
