[JWT] Implementing JWT verification

Implementing JsonWebToken verification that ensures:
- algorithms match
- signature is valid
- issuers match
- token didn't expire yet
- token isn't used before it should

diff --git a/Engine/Plugins/Experimental/JWT/Source/JWT/Private/JsonWebToken.cpp b/Engine/Plugins/Experimental/JWT/Source/JWT/Private/JsonWebToken.cpp
--- a/Engine/Plugins/Experimental/JWT/Source/JWT/Private/JsonWebToken.cpp
+++ b/Engine/Plugins/Experimental/JWT/Source/JWT/Private/JsonWebToken.cpp
@@ -4,6 +4,7 @@
 
 #include "JwtGlobals.h"
 #include "JwtUtils.h"
+#include "JwtAlgorithms.h"
 
 #include "Algo/Count.h"
 #include "Serialization/JsonSerializer.h"
@@ -269,13 +270,117 @@ bool FJsonWebToken::GetClaim(const FStringView InName, TSharedPtr<FJsonValue>& O
 }
 
 
-/**
- * Verify the JWT's signature.
- *
- * TODO(vri): Implement JWT verification with public key
- */
 bool FJsonWebToken::Verify() const
 {
 	UE_LOG(LogJwt, Error, TEXT("[FJsonWebToken::Verify] JWT signature verification is not implemented yet and will always return false."));
 	return false;
 }
+
+
+bool FJsonWebToken::Verify(
+	const IJwtAlgorithm& Algorithm, const FStringView ExpectedIssuer) const
+{
+	// Check whether the signature is set
+	if (!Signature.IsSet())
+	{
+		UE_LOG(LogJwt, Error,
+			TEXT("[FJsonWebToken::Verify] No signature to verify."));
+
+		return false;
+	}
+
+	FString IndicatedAlgorithm;
+
+	// Check whether the algorithm is set
+	if (!GetAlgorithm(IndicatedAlgorithm))
+	{
+		UE_LOG(LogJwt, Error,
+			TEXT("[FJsonWebToken::Verify] Could not get token's algorithm."));
+
+		return false;
+	}
+
+	// Check whether the algorithms match
+	if (Algorithm.GetAlgString() != IndicatedAlgorithm)
+	{
+		UE_LOG(LogJwt, Error,
+			TEXT("[FJsonWebToken::Verify] Algorithms don't match."));
+
+		return false;
+	}
+
+	TArray<uint8> EncodedHeaderPayloadBytes;
+	FJwtUtils::StringViewToBytes(EncodedHeaderPayload, EncodedHeaderPayloadBytes);
+
+	// Verify the signature
+	if (!Algorithm.VerifySignature(EncodedHeaderPayloadBytes, Signature.GetValue()))
+	{
+		UE_LOG(LogJwt, Error,
+			TEXT("[FJsonWebToken::Verify] Signature verification failed."));
+
+		return false;
+	}
+
+	FString Issuer;
+
+	// Check whether the issuer is set
+	if (!GetIssuer(Issuer))
+	{
+		UE_LOG(LogJwt, Error,
+			TEXT("[FJsonWebToken::Verify] Issuer not set."));
+
+		return false;
+	}
+
+	// Check whether the issuers match
+	if (ExpectedIssuer != FStringView(Issuer))
+	{
+		UE_LOG(LogJwt, Error,
+			TEXT("[FJsonWebToken::Verify] Issuer does not match expected issuer."));
+
+		return false;
+	}
+
+	int64 IssuedAt = 0, Expiration = 0;
+
+	// Check whether IssuedAt and Expiration timestamps are set
+	if (!GetIssuedAt(IssuedAt) || !GetExpiration(Expiration))
+	{
+		UE_LOG(LogJwt, Error,
+			TEXT("[FJsonWebToken::Verify] IssuedAt or Expiration timestamp is not set."));
+
+		return false;
+	}
+
+	const FDateTime Now = FDateTime::UtcNow();
+	const FDateTime TimeIssuedAt = FDateTime::FromUnixTimestamp(IssuedAt);
+	const FDateTime TimeExpires = FDateTime::FromUnixTimestamp(Expiration);
+
+	// Check whether the token has expired or is invalid
+	if (TimeIssuedAt >= Now || TimeExpires <= Now)
+	{
+		UE_LOG(LogJwt, Error,
+			TEXT("[FJsonWebToken::Verify] Token not valid or has expired already."));
+
+		return false;
+	}
+
+	int64 NotBefore = 0;
+
+	// Check whether the token is used before NotBefore, if it's set
+	if (GetNotBefore(NotBefore))
+	{
+		const FDateTime TimeNotBefore = FDateTime::FromUnixTimestamp(NotBefore);
+
+		if (TimeNotBefore >= Now)
+		{
+			UE_LOG(LogJwt, Error,
+				TEXT("[FJsonWebToken::Verify] Token not valid yet."));
+
+			return false;
+		}
+	}
+
+	return true;
+}
+
diff --git a/Engine/Plugins/Experimental/JWT/Source/JWT/Public/JsonWebToken.h b/Engine/Plugins/Experimental/JWT/Source/JWT/Public/JsonWebToken.h
--- a/Engine/Plugins/Experimental/JWT/Source/JWT/Public/JsonWebToken.h
+++ b/Engine/Plugins/Experimental/JWT/Source/JWT/Public/JsonWebToken.h
@@ -6,6 +6,7 @@
 
 class FJsonObject;
 class FJsonValue;
+class IJwtAlgorithm;
 enum class EJson;
 
 
@@ -122,15 +123,27 @@ public:
 	}
 
 	/**
-	 * Perform verification of the encoded header and encoded payload.
-	 * The is done using the cryptographic algorithm specified in the header, and a secret value obtained
-	 * by key-lookup, to generate a signature.  The generated signature must be identical to the signature
-	 * provided in the JWT for verification to succeed.
-	 * 
-	 * @return true if the header and payload were verified successfully (signature match), otherwise false.
+	 * Deprecated method to signature validate the JWT.
+	 *
+	 * @return False
 	 */
+	UE_DEPRECATED(5.3, "Verify() without arguments is deprecated. Please use Verify(Algorithm, ExpectedIssuer) instead.")
 	bool Verify() const;
 
+	/**
+	 * Signature validate and verify the JWT.
+	 * - Validates the signature against the encoded header and encoded payload
+	 * - Verifies the basic claims of the JWT
+	 * - Ensures the issuers match
+	 *
+	 * @param Algorithm Implementation of the cryptographic algorithm used for signature validation
+	 * @param ExpectedIssuer The expected issuer
+	 *
+	 * @return Whether the JWT was successfully verified
+	 */
+	bool Verify(
+		const IJwtAlgorithm& Algorithm, const FStringView ExpectedIssuer) const;
+
 private:
 	FJsonWebToken(
 		const FStringView InEncodedHeaderPayload, const TSharedRef<FJsonObject>& InHeaderPtr,
diff --git a/Engine/Plugins/Experimental/JWT/Source/JWT/Public/JwtAlgorithm.h b/Engine/Plugins/Experimental/JWT/Source/JWT/Public/JwtAlgorithm.h
--- a/Engine/Plugins/Experimental/JWT/Source/JWT/Public/JwtAlgorithm.h
+++ b/Engine/Plugins/Experimental/JWT/Source/JWT/Public/JwtAlgorithm.h
@@ -21,7 +21,7 @@ public:
 	/**
 	 * Gets the algorithm's string representation.
 	 */
-	virtual inline const FString& GetAlgString() = 0;
+	virtual inline const FString& GetAlgString() const = 0;
 
 	/**
 	 * Verifies the decoded signature with the encoded message.
diff --git a/Engine/Plugins/Experimental/JWT/Source/JWT/Public/JwtAlgorithms.h b/Engine/Plugins/Experimental/JWT/Source/JWT/Public/JwtAlgorithms.h
--- a/Engine/Plugins/Experimental/JWT/Source/JWT/Public/JwtAlgorithms.h
+++ b/Engine/Plugins/Experimental/JWT/Source/JWT/Public/JwtAlgorithms.h
@@ -23,7 +23,7 @@ public:
 
 	// ~ IJwtAlgorithm interface
 
-	virtual inline const FString& GetAlgString() override
+	virtual inline const FString& GetAlgString() const override
 	{
 		return AlgorithmString;
 	}
