[JWT] Centralizing GetEncryptionContext UniquePtr in FJwtAlgorithm_RS256 constructor

Centralizing the call to GetEncryptionContext that creates a new context of OpenSSL/SwitchSSL PlatformCrypto implementation once in the constructor of FJwtAlgorithm_RS256, to avoid doing it multiple times.

diff --git a/Engine/Plugins/Experimental/JWT/Source/JWT/JWT.Build.cs b/Engine/Plugins/Experimental/JWT/Source/JWT/JWT.Build.cs
--- a/Engine/Plugins/Experimental/JWT/Source/JWT/JWT.Build.cs
+++ b/Engine/Plugins/Experimental/JWT/Source/JWT/JWT.Build.cs
@@ -9,13 +9,7 @@ namespace UnrealBuildTool.Rules
 			PublicDependencyModuleNames.AddRange(
 				new string[]
 				{
-					"Core"
-				}
-			);
-
-			PrivateDependencyModuleNames.AddRange(
-				new string[]
-				{
+					"Core",
 					"Json",
 					"PlatformCrypto",
 					"PlatformCryptoOpenSSL"
diff --git a/Engine/Plugins/Experimental/JWT/Source/JWT/Private/JwtAlgorithms.cpp b/Engine/Plugins/Experimental/JWT/Source/JWT/Private/JwtAlgorithms.cpp
--- a/Engine/Plugins/Experimental/JWT/Source/JWT/Private/JwtAlgorithms.cpp
+++ b/Engine/Plugins/Experimental/JWT/Source/JWT/Private/JwtAlgorithms.cpp
@@ -2,7 +2,11 @@
 
 #include "JwtAlgorithms.h"
 
-#include "JwtUtils.h"
+
+FJwtAlgorithm_RS256::FJwtAlgorithm_RS256()
+{
+	EncryptionContext = FJwtUtils::GetEncryptionContext();
+}
 
 
 FJwtAlgorithm_RS256::~FJwtAlgorithm_RS256()
@@ -18,9 +22,6 @@ bool FJwtAlgorithm_RS256::VerifySignature(
 	const TArrayView<const uint8> EncodedMessage,
 	const TArrayView<const uint8> DecodedSignature) const
 {
-	TUniquePtr<FEncryptionContext> EncryptionContext
-		= FJwtUtils::GetEncryptionContext();
-
 	if (!EncryptionContext.IsValid())
 	{
 		return false;
@@ -41,9 +42,6 @@ bool FJwtAlgorithm_RS256::VerifySignature(
 
 bool FJwtAlgorithm_RS256::SetPublicKey(const TArrayView<const uint8> InKey)
 {
-	TUniquePtr<FEncryptionContext> EncryptionContext
-		= FJwtUtils::GetEncryptionContext();
-
 	if (!EncryptionContext.IsValid())
 	{
 		return false;
@@ -80,9 +78,6 @@ bool FJwtAlgorithm_RS256::SetPublicKey(const FStringView InKey)
 
 void FJwtAlgorithm_RS256::DestroyKey(void* Key)
 {
-	TUniquePtr<FEncryptionContext> EncryptionContext
-		= FJwtUtils::GetEncryptionContext();
-
 	if (!EncryptionContext.IsValid())
 	{
 		return;
diff --git a/Engine/Plugins/Experimental/JWT/Source/JWT/Public/JwtAlgorithms.h b/Engine/Plugins/Experimental/JWT/Source/JWT/Public/JwtAlgorithms.h
--- a/Engine/Plugins/Experimental/JWT/Source/JWT/Public/JwtAlgorithms.h
+++ b/Engine/Plugins/Experimental/JWT/Source/JWT/Public/JwtAlgorithms.h
@@ -4,6 +4,8 @@
 
 #include "JwtAlgorithm.h"
 
+#include "JwtUtils.h"
+
 
 class JWT_API FJwtAlgorithm_RS256
 	: public IJwtAlgorithm
@@ -11,7 +13,7 @@ class JWT_API FJwtAlgorithm_RS256
 
 public:
 
-	FJwtAlgorithm_RS256() = default;
+	FJwtAlgorithm_RS256();
 
 	FJwtAlgorithm_RS256(const FJwtAlgorithm_RS256&) = delete;
 
@@ -68,6 +70,9 @@ protected:
 
 private:
 
+	/** Holds the unique pointer to the encryption context (OpenSSL/SwitchSSL). */
+	TUniquePtr<FEncryptionContext> EncryptionContext = nullptr;
+
 	/** Holds the public key pointer. */
 	void* PublicKey = nullptr;
 
