[JWT] Adding method to check if JsonWebToken has expired

Adding a public convenience method to check if a JsonWebToken has expired.

diff --git a/Engine/Plugins/Experimental/JWT/Source/JWT/Private/JsonWebToken.cpp b/Engine/Plugins/Experimental/JWT/Source/JWT/Private/JsonWebToken.cpp
--- a/Engine/Plugins/Experimental/JWT/Source/JWT/Private/JsonWebToken.cpp
+++ b/Engine/Plugins/Experimental/JWT/Source/JWT/Private/JsonWebToken.cpp
@@ -270,6 +270,24 @@ bool FJsonWebToken::GetClaim(const FStringView InName, TSharedPtr<FJsonValue>& O
 }
 
 
+bool FJsonWebToken::HasExpired() const
+{
+	int64 ExpirationTimestamp = 0;
+
+	if (!GetExpiration(ExpirationTimestamp))
+	{
+		UE_LOG(LogJwt, Error, TEXT("[FJsonWebToken::HasExpired] Could not get expiration timestamp."));
+
+		return false;
+	}
+
+	const FDateTime Now = FDateTime::UtcNow();
+	const FDateTime TimeExpires = FDateTime::FromUnixTimestamp(ExpirationTimestamp);
+
+	return Now > TimeExpires;
+}
+
+
 bool FJsonWebToken::Verify() const
 {
 	UE_LOG(LogJwt, Error, TEXT("[FJsonWebToken::Verify] JWT signature verification is not implemented yet and will always return false."));
@@ -354,10 +372,9 @@ bool FJsonWebToken::Verify(
 
 	const FDateTime Now = FDateTime::UtcNow();
 	const FDateTime TimeIssuedAt = FDateTime::FromUnixTimestamp(IssuedAt);
-	const FDateTime TimeExpires = FDateTime::FromUnixTimestamp(Expiration);
 
 	// Check whether the token has expired or is invalid
-	if (TimeIssuedAt >= Now || TimeExpires <= Now)
+	if (TimeIssuedAt >= Now || HasExpired())
 	{
 		UE_LOG(LogJwt, Error,
 			TEXT("[FJsonWebToken::Verify] Token not valid or has expired already."));
diff --git a/Engine/Plugins/Experimental/JWT/Source/JWT/Public/JsonWebToken.h b/Engine/Plugins/Experimental/JWT/Source/JWT/Public/JsonWebToken.h
--- a/Engine/Plugins/Experimental/JWT/Source/JWT/Public/JsonWebToken.h
+++ b/Engine/Plugins/Experimental/JWT/Source/JWT/Public/JsonWebToken.h
@@ -122,6 +122,13 @@ public:
 		return Payload->GetField<JsonType>(FString(InName));
 	}
 
+	/**
+	 * Checks whether the tokens expiration timestamp is in the past.
+	 *
+	 * @return Whether the token has expired
+	 */
+	bool HasExpired() const;
+
 	/**
 	 * Deprecated method to signature validate the JWT.
 	 *
