Add new functtion to get custom string claims in JWT module.

diff --git a/Engine/Plugins/Experimental/JWT/Source/JWT/Private/JsonWebToken.cpp b/Engine/Plugins/Experimental/JWT/Source/JWT/Private/JsonWebToken.cpp
--- a/Engine/Plugins/Experimental/JWT/Source/JWT/Private/JsonWebToken.cpp
+++ b/Engine/Plugins/Experimental/JWT/Source/JWT/Private/JsonWebToken.cpp
@@ -263,6 +263,11 @@ bool FJsonWebToken::GetAudience(FString& OutValue) const
 	return Payload->TryGetStringField(CLAIM_AUDIENCE, OutValue);
 }
 
+bool FJsonWebToken::GetStringClaim(const FString& InName, FString& OutValue) const
+{
+	return Payload->TryGetStringField(InName, OutValue);
+}
+
 bool FJsonWebToken::GetClaim(const FStringView InName, TSharedPtr<FJsonValue>& OutClaim) const
 {
 	OutClaim = Payload->TryGetField(FString(InName));
diff --git a/Engine/Plugins/Experimental/JWT/Source/JWT/Public/JsonWebToken.h b/Engine/Plugins/Experimental/JWT/Source/JWT/Public/JsonWebToken.h
--- a/Engine/Plugins/Experimental/JWT/Source/JWT/Public/JsonWebToken.h
+++ b/Engine/Plugins/Experimental/JWT/Source/JWT/Public/JsonWebToken.h
@@ -99,6 +99,15 @@ public:
 	 */
 	bool GetAudience(FString& OutValue) const;
 
+	/**
+	 * Get a custom claim by name.
+	 * 
+	 * @param InName The name of the claim.
+	 * @param OutValue The value to output on success.
+	 * @return Whether the claim exists and was successfully outputted.
+	 */
+	bool GetStringClaim(const FString& InName, FString& OutValue) const;
+
 	/**
 	 * Gets a claim by name.
 	 * This method can be used to get custom claims that are not reserved as part of the JWT specification.
