Add JWT option to get raw JSON payload.

diff --git a/Engine/Plugins/Experimental/JWT/Source/JWT/Public/JsonWebToken.h b/Engine/Plugins/Experimental/JWT/Source/JWT/Public/JsonWebToken.h
--- a/Engine/Plugins/Experimental/JWT/Source/JWT/Public/JsonWebToken.h
+++ b/Engine/Plugins/Experimental/JWT/Source/JWT/Public/JsonWebToken.h
@@ -27,6 +27,13 @@ public:
 
 	static bool FromString(const FStringView InEncodedJsonWebToken, FJsonWebToken& OutJsonWebToken);
 
+	/**
+	 * Get raw JSON object for payload, allowing for custom claim parsing with FJsonSerializable classes.
+	 * 
+	 * @return reference to FJsonObject for the payload.
+	 */
+	const TSharedRef<FJsonObject>& GetPayload() const { return Payload; }
+
 	/**
 	 * Gets the type.
 	 *
